Add Source Code Access plugin for the 10X editor
* Had to use N10X in some places due to code restrictions
* Using Hardcoded path for now

#tests Double clicked nodes to have source code opened properly.
#pf 6481108de8609a90465e556d

diff --git a/Engine/Plugins/Developer/N10XSourceCodeAccess/10XSourceCodeAccess.uplugin b/Engine/Plugins/Developer/N10XSourceCodeAccess/10XSourceCodeAccess.uplugin
new file mode 100644
--- /dev/null
+++ b/Engine/Plugins/Developer/N10XSourceCodeAccess/10XSourceCodeAccess.uplugin
@@ -0,0 +1,29 @@
+{
+	"FileVersion" : 3,
+	"Version" : 1,
+	"VersionName" : "1.0",
+	"FriendlyName" : "10X Editor Integration",
+	"Description" : "Allows access to source code in the 10X Editor .",
+	"Category" : "Programming",
+	"CreatedBy" : "Epic Games, Inc.",
+	"CreatedByURL" : "https://epicgames.com",
+	"DocsURL" : "",
+	"MarketplaceURL" : "",
+	"SupportURL" : "",
+	"EnabledByDefault" : true,
+	"CanContainContent" : false,
+	"IsBetaVersion" : false,
+	"Installed" : false,
+	"Modules" :
+	[
+		{
+			"Name" : "N10XSourceCodeAccess",
+			"Type" : "UncookedOnly",
+			"LoadingPhase" : "Default",
+			"PlatformAllowList" :
+			[
+				"Win64"
+			]
+		}
+	]
+}
diff --git a/Engine/Plugins/Developer/N10XSourceCodeAccess/Source/N10XSourceCodeAccess/N10XSourceCodeAccess.Build.cs b/Engine/Plugins/Developer/N10XSourceCodeAccess/Source/N10XSourceCodeAccess/N10XSourceCodeAccess.Build.cs
new file mode 100644
--- /dev/null
+++ b/Engine/Plugins/Developer/N10XSourceCodeAccess/Source/N10XSourceCodeAccess/N10XSourceCodeAccess.Build.cs
@@ -0,0 +1,31 @@
+// Copyright Epic Games, Inc. All Rights Reserved.
+using System;
+using Microsoft.Win32;
+
+namespace UnrealBuildTool.Rules
+{
+	// The IDE is called 10X, we add the N to make it compiler friendly
+	public class N10XSourceCodeAccess : ModuleRules
+	{
+        public N10XSourceCodeAccess(ReadOnlyTargetRules Target) : base(Target)
+		{
+            PrivateDependencyModuleNames.AddRange(
+				new string[]
+				{
+					"Core",
+					"SourceCodeAccess",
+					"DesktopPlatform",
+				}
+			);
+
+			if (Target.bBuildEditor)
+			{
+				PrivateDependencyModuleNames.Add("HotReload");
+			}
+
+			bBuildLocallyWithSNDBS = true;
+
+			ShortName = "10XSCA";
+		}
+	}
+}
diff --git a/Engine/Plugins/Developer/N10XSourceCodeAccess/Source/N10XSourceCodeAccess/Private/10XSourceCodeAccessModule.cpp b/Engine/Plugins/Developer/N10XSourceCodeAccess/Source/N10XSourceCodeAccess/Private/10XSourceCodeAccessModule.cpp
new file mode 100644
--- /dev/null
+++ b/Engine/Plugins/Developer/N10XSourceCodeAccess/Source/N10XSourceCodeAccess/Private/10XSourceCodeAccessModule.cpp
@@ -0,0 +1,37 @@
+// Copyright Epic Games, Inc. All Rights Reserved.
+
+#include "10XSourceCodeAccessModule.h"
+#include "Features/IModularFeatures.h"
+#include "Modules/ModuleManager.h"
+#include "HAL/LowLevelMemTracker.h"
+#include "10XSourceCodeAccessor.h"
+
+LLM_DEFINE_TAG(10XSourceCodeAccess);
+
+IMPLEMENT_MODULE( F10XSourceCodeAccessModule, 10XSourceCodeAccess );
+
+F10XSourceCodeAccessModule::F10XSourceCodeAccessModule()
+	: SourceCodeAccessor(MakeShareable(new F10XSourceCodeAccessor()))
+{
+}
+
+void F10XSourceCodeAccessModule::StartupModule()
+{
+	LLM_SCOPE_BYTAG(10XSourceCodeAccess);
+
+	SourceCodeAccessor->Startup();
+
+	IModularFeatures::Get().RegisterModularFeature(TEXT("SourceCodeAccessor"), &SourceCodeAccessor.Get() );
+}
+
+void F10XSourceCodeAccessModule::ShutdownModule()
+{
+	IModularFeatures::Get().UnregisterModularFeature(TEXT("SourceCodeAccessor"), &SourceCodeAccessor.Get());
+
+	SourceCodeAccessor->Shutdown();
+}
+
+F10XSourceCodeAccessor& F10XSourceCodeAccessModule::GetAccessor()
+{
+	return SourceCodeAccessor.Get();
+}
diff --git a/Engine/Plugins/Developer/N10XSourceCodeAccess/Source/N10XSourceCodeAccess/Private/10XSourceCodeAccessModule.h b/Engine/Plugins/Developer/N10XSourceCodeAccess/Source/N10XSourceCodeAccess/Private/10XSourceCodeAccessModule.h
new file mode 100644
--- /dev/null
+++ b/Engine/Plugins/Developer/N10XSourceCodeAccess/Source/N10XSourceCodeAccess/Private/10XSourceCodeAccessModule.h
@@ -0,0 +1,24 @@
+﻿// Copyright Epic Games, Inc. All Rights Reserved.
+
+#pragma once
+
+#include "Modules/ModuleInterface.h"
+#include "Templates/SharedPointer.h"
+
+class F10XSourceCodeAccessor;
+
+class F10XSourceCodeAccessModule : public IModuleInterface
+{
+public:
+	F10XSourceCodeAccessModule();
+
+	/** IModuleInterface implementation */
+	virtual void StartupModule() override;
+	virtual void ShutdownModule() override;
+
+	F10XSourceCodeAccessor& GetAccessor();
+
+private:
+
+	TSharedRef<F10XSourceCodeAccessor> SourceCodeAccessor;
+};
diff --git a/Engine/Plugins/Developer/N10XSourceCodeAccess/Source/N10XSourceCodeAccess/Private/10XSourceCodeAccessor.cpp b/Engine/Plugins/Developer/N10XSourceCodeAccess/Source/N10XSourceCodeAccess/Private/10XSourceCodeAccessor.cpp
new file mode 100644
--- /dev/null
+++ b/Engine/Plugins/Developer/N10XSourceCodeAccess/Source/N10XSourceCodeAccess/Private/10XSourceCodeAccessor.cpp
@@ -0,0 +1,214 @@
+// Copyright Epic Games, Inc. All Rights Reserved.
+
+#include "10XSourceCodeAccessor.h"
+#include "Misc/Paths.h"
+#include "Misc/UProjectInfo.h"
+#include "Misc/App.h"
+
+#if PLATFORM_WINDOWS
+#include "Internationalization/Regex.h"
+#include "Windows/AllowWindowsPlatformTypes.h"
+#endif
+
+DEFINE_LOG_CATEGORY_STATIC(LogVSCodeAccessor, Log, All);
+
+#define LOCTEXT_NAMESPACE "10XSourceCodeAccessor"
+
+namespace
+{
+	static const TCHAR* G10XWorkspaceExtension = TEXT(".sln");
+}
+
+const FString& F10XSourceCodeAccessor::GetSolutionPath() const
+{
+	FScopeLock Lock(&CachedSolutionPathCriticalSection);
+
+	if (CachedSolutionPath.IsEmpty() && IsInGameThread())
+	{
+		CachedSolutionPath = FPaths::ProjectDir();
+
+		if (!FUProjectDictionary::GetDefault().IsForeignProject(CachedSolutionPath))
+		{
+			CachedSolutionPath = FPaths::Combine(FPaths::RootDir(), FString("UE5") + G10XWorkspaceExtension);
+		}
+		else
+		{
+			FString BaseName = FApp::HasProjectName() ? FApp::GetProjectName() : FPaths::GetBaseFilename(CachedSolutionPath);
+			CachedSolutionPath = FPaths::Combine(CachedSolutionPath, BaseName + G10XWorkspaceExtension);
+		}
+	}
+
+	return CachedSolutionPath;
+}
+
+void F10XSourceCodeAccessor::Startup()
+{
+	GetSolutionPath();
+	RefreshAvailability();
+}
+
+void F10XSourceCodeAccessor::RefreshAvailability()
+{
+#if PLATFORM_WINDOWS
+	FString IDEPath;
+
+	const TCHAR* ClassesKey = TEXT("SOFTWARE\\Classes\\PureDevSoftware.10x.1\\shell\\open\\");
+	const TCHAR* InstallDirKey = TEXT("SOFTWARE\\PureDevSoftware\\10x\\");
+
+	// Check for explorer actions first since it will have full path including exe
+	if (FWindowsPlatformMisc::QueryRegKey(HKEY_CURRENT_USER, ClassesKey, TEXT(""), IDEPath) ||
+		FWindowsPlatformMisc::QueryRegKey(HKEY_LOCAL_MACHINE, ClassesKey, TEXT(""), IDEPath))
+	{
+		// Matches shell command "App" "%1" to extract App
+		FString PatternString(TEXT("\"(.*)\" \".*\""));
+		FRegexPattern Pattern(PatternString);
+		FRegexMatcher Matcher(Pattern, IDEPath);
+		if (Matcher.FindNext())
+		{
+			IDEPath = Matcher.GetCaptureGroup(1);
+		}
+	}
+	// Fall back to install directory and add 10x.exe to resulting path if user didn't install explorer integrations
+	else if (FWindowsPlatformMisc::QueryRegKey(HKEY_CURRENT_USER, InstallDirKey, TEXT("InstallDir"), IDEPath) ||
+		FWindowsPlatformMisc::QueryRegKey(HKEY_LOCAL_MACHINE, InstallDirKey, TEXT("InstallDir"), IDEPath))
+	{
+		IDEPath = FPaths::Combine(IDEPath, TEXT("10x.exe"));
+	}
+
+	if (FPaths::FileExists(IDEPath))
+	{
+		ApplicationFilePath = IDEPath;
+	}
+#endif
+}
+
+void F10XSourceCodeAccessor::Shutdown()
+{
+}
+
+bool F10XSourceCodeAccessor::OpenSourceFiles(const TArray<FString>& AbsoluteSourcePaths)
+{
+	if (CanAccessSourceCode())
+	{
+		FString SolutionDir = GetSolutionPath();
+		TArray<FString> Args;
+
+		for (const FString& SourcePath : AbsoluteSourcePaths)
+		{
+			Args.Add(SourcePath);
+		}
+
+		return Launch(Args);
+	}
+
+	return false;
+}
+
+bool F10XSourceCodeAccessor::AddSourceFiles(const TArray<FString>& AbsoluteSourcePaths, const TArray<FString>& AvailableModules)
+{
+	// dosn't need to do anything when new files are added
+	return true;
+}
+
+bool F10XSourceCodeAccessor::OpenFileAtLine(const FString& FullPath, int32 LineNumber, int32 ColumnNumber)
+{
+	if (CanAccessSourceCode())
+	{
+		// Column & line numbers are 1-based, 10X is 0 based (like it should be!)
+		LineNumber = FMath::Max(LineNumber - 1, 0);
+		ColumnNumber = FMath::Max(ColumnNumber - 1, 0);
+
+		TArray<FString> Args;
+		Args.Add(GetSolutionPath());
+		Args.Add(FullPath);
+		Args.Add(FString::Printf(TEXT("N10X.Editor.SetCursorPos((%d,%d))"), ColumnNumber, LineNumber));
+		return Launch(Args);
+	}
+
+	return false;
+}
+
+bool F10XSourceCodeAccessor::CanAccessSourceCode() const
+{
+	// True if we have any discovered the application
+	return !ApplicationFilePath.IsEmpty();
+}
+
+FName F10XSourceCodeAccessor::GetFName() const
+{
+	return FName("10X Editor");
+}
+
+FText F10XSourceCodeAccessor::GetNameText() const
+{
+	return LOCTEXT("10XDisplayName", "10X Editor");
+}
+
+FText F10XSourceCodeAccessor::GetDescriptionText() const
+{
+	return LOCTEXT("10XDisplayDesc", "Open source code files in the 10X Editor");
+}
+
+void F10XSourceCodeAccessor::Tick(const float DeltaTime)
+{
+}
+
+bool F10XSourceCodeAccessor::OpenSolution()
+{
+	if (CanAccessSourceCode())
+	{
+		return OpenSolutionAtPath(GetSolutionPath());
+	}
+
+	return false;
+}
+
+bool F10XSourceCodeAccessor::OpenSolutionAtPath(const FString& InSolutionPath)
+{
+	if (CanAccessSourceCode())
+	{
+		FString SolutionPath = InSolutionPath;
+
+		if (!SolutionPath.EndsWith(G10XWorkspaceExtension))
+		{
+			SolutionPath = SolutionPath + G10XWorkspaceExtension;
+		}
+
+		TArray<FString> Args;
+		Args.Add(SolutionPath);
+		return Launch(Args);
+	}
+
+	return false;
+}
+
+bool F10XSourceCodeAccessor::DoesSolutionExist() const
+{
+	return FPaths::FileExists(GetSolutionPath());
+}
+
+bool F10XSourceCodeAccessor::SaveAllOpenDocuments() const
+{
+	return true;
+}
+
+bool F10XSourceCodeAccessor::Launch(const TArray<FString>& InArgs)
+{
+	if (CanAccessSourceCode())
+	{
+		FString ArgsString;
+		for (const FString& Arg : InArgs)
+		{
+			ArgsString.Append(Arg);
+			ArgsString.Append(TEXT(" "));
+		}
+
+		uint32 ProcessID;
+		FProcHandle hProcess = FPlatformProcess::CreateProc(*ApplicationFilePath, *ArgsString, true, false, false, &ProcessID, 0, nullptr, nullptr, nullptr);
+		return hProcess.IsValid();
+	}
+	
+	return false;
+}
+
+#undef LOCTEXT_NAMESPACE
diff --git a/Engine/Plugins/Developer/N10XSourceCodeAccess/Source/N10XSourceCodeAccess/Private/10XSourceCodeAccessor.h b/Engine/Plugins/Developer/N10XSourceCodeAccess/Source/N10XSourceCodeAccess/Private/10XSourceCodeAccessor.h
new file mode 100644
--- /dev/null
+++ b/Engine/Plugins/Developer/N10XSourceCodeAccess/Source/N10XSourceCodeAccess/Private/10XSourceCodeAccessor.h
@@ -0,0 +1,40 @@
+// Copyright Epic Games, Inc. All Rights Reserved.
+#pragma once
+
+#include "ISourceCodeAccessor.h"
+
+class F10XSourceCodeAccessor : public ISourceCodeAccessor
+{
+public:
+	F10XSourceCodeAccessor()
+	{
+	}
+
+	void Startup();
+	void Shutdown();
+
+	/** ISourceCodeAccessor implementation */
+	virtual void RefreshAvailability() final;
+	virtual bool CanAccessSourceCode() const final;
+	virtual FName GetFName() const final;
+	virtual FText GetNameText() const final;
+	virtual FText GetDescriptionText() const final;
+	virtual bool OpenSolution() final;
+	virtual bool OpenSolutionAtPath(const FString& InSolutionPath) final;
+	virtual bool DoesSolutionExist() const final;
+	virtual bool OpenFileAtLine(const FString& FullPath, int32 LineNumber, int32 ColumnNumber = 0) final;
+	virtual bool OpenSourceFiles(const TArray<FString>& AbsoluteSourcePaths) final;
+	virtual bool AddSourceFiles(const TArray<FString>& AbsoluteSourcePaths, const TArray<FString>& AvailableModules) final;
+	virtual bool SaveAllOpenDocuments() const final;
+	virtual void Tick(const float DeltaTime) final;
+
+private:
+	FString ApplicationFilePath;
+
+	mutable FString CachedSolutionPath;
+	mutable FCriticalSection CachedSolutionPathCriticalSection;
+
+	const FString& GetSolutionPath() const;
+
+	bool Launch(const TArray<FString>& InArgs);
+};
