Diabling plugin by default as the module N10XSourceCodeAccess failes to load
Changing LLM Tag name as it can't start with an integer

diff --git a/Engine/Plugins/Developer/N10XSourceCodeAccess/10XSourceCodeAccess.uplugin b/Engine/Plugins/Developer/N10XSourceCodeAccess/10XSourceCodeAccess.uplugin
--- a/Engine/Plugins/Developer/N10XSourceCodeAccess/10XSourceCodeAccess.uplugin
+++ b/Engine/Plugins/Developer/N10XSourceCodeAccess/10XSourceCodeAccess.uplugin
@@ -10,7 +10,7 @@
 	"DocsURL" : "",
 	"MarketplaceURL" : "",
 	"SupportURL" : "",
-	"EnabledByDefault" : true,
+	"EnabledByDefault" : false,
 	"CanContainContent" : false,
 	"IsBetaVersion" : false,
 	"Installed" : false,
diff --git a/Engine/Plugins/Developer/N10XSourceCodeAccess/Source/N10XSourceCodeAccess/Private/10XSourceCodeAccessModule.cpp b/Engine/Plugins/Developer/N10XSourceCodeAccess/Source/N10XSourceCodeAccess/Private/10XSourceCodeAccessModule.cpp
--- a/Engine/Plugins/Developer/N10XSourceCodeAccess/Source/N10XSourceCodeAccess/Private/10XSourceCodeAccessModule.cpp
+++ b/Engine/Plugins/Developer/N10XSourceCodeAccess/Source/N10XSourceCodeAccess/Private/10XSourceCodeAccessModule.cpp
@@ -6,7 +6,7 @@
 #include "HAL/LowLevelMemTracker.h"
 #include "10XSourceCodeAccessor.h"
 
-LLM_DEFINE_TAG(10XSourceCodeAccess);
+LLM_DEFINE_TAG(XSourceCodeAccess_10);
 
 IMPLEMENT_MODULE( F10XSourceCodeAccessModule, 10XSourceCodeAccess );
 
@@ -17,7 +17,7 @@ F10XSourceCodeAccessModule::F10XSourceCodeAccessModule()
 
 void F10XSourceCodeAccessModule::StartupModule()
 {
-	LLM_SCOPE_BYTAG(10XSourceCodeAccess);
+	LLM_SCOPE_BYTAG(XSourceCodeAccess_10);
 
 	SourceCodeAccessor->Startup();
 
