Add letter prefix to llm tag, re-enable 10x source accessor plugin

#pf 6482963cf401ef063679b4c4

diff --git a/Engine/Plugins/Developer/N10XSourceCodeAccess/10XSourceCodeAccess.uplugin b/Engine/Plugins/Developer/N10XSourceCodeAccess/10XSourceCodeAccess.uplugin
--- a/Engine/Plugins/Developer/N10XSourceCodeAccess/10XSourceCodeAccess.uplugin
+++ b/Engine/Plugins/Developer/N10XSourceCodeAccess/10XSourceCodeAccess.uplugin
@@ -10,7 +10,7 @@
 	"DocsURL" : "",
 	"MarketplaceURL" : "",
 	"SupportURL" : "",
-	"EnabledByDefault" : false,
+	"EnabledByDefault" : true,
 	"CanContainContent" : false,
 	"IsBetaVersion" : false,
 	"Installed" : false,
diff --git a/Engine/Plugins/Developer/N10XSourceCodeAccess/Source/N10XSourceCodeAccess/Private/10XSourceCodeAccessModule.cpp b/Engine/Plugins/Developer/N10XSourceCodeAccess/Source/N10XSourceCodeAccess/Private/10XSourceCodeAccessModule.cpp
--- a/Engine/Plugins/Developer/N10XSourceCodeAccess/Source/N10XSourceCodeAccess/Private/10XSourceCodeAccessModule.cpp
+++ b/Engine/Plugins/Developer/N10XSourceCodeAccess/Source/N10XSourceCodeAccess/Private/10XSourceCodeAccessModule.cpp
@@ -6,7 +6,7 @@
 #include "HAL/LowLevelMemTracker.h"
 #include "10XSourceCodeAccessor.h"
 
-LLM_DEFINE_TAG(XSourceCodeAccess_10);
+LLM_DEFINE_TAG(N10XSourceCodeAccess);
 
 IMPLEMENT_MODULE( F10XSourceCodeAccessModule, 10XSourceCodeAccess );
 
@@ -17,7 +17,7 @@ F10XSourceCodeAccessModule::F10XSourceCodeAccessModule()
 
 void F10XSourceCodeAccessModule::StartupModule()
 {
-	LLM_SCOPE_BYTAG(XSourceCodeAccess_10);
+	LLM_SCOPE_BYTAG(N10XSourceCodeAccess);
 
 	SourceCodeAccessor->Startup();
 
