Fix cis error

diff --git a/Engine/Plugins/Developer/N10XSourceCodeAccess/Source/N10XSourceCodeAccess/Private/N10XSourceCodeAccessModule.cpp b/Engine/Plugins/Developer/N10XSourceCodeAccess/Source/N10XSourceCodeAccess/Private/N10XSourceCodeAccessModule.cpp
--- a/Engine/Plugins/Developer/N10XSourceCodeAccess/Source/N10XSourceCodeAccess/Private/N10XSourceCodeAccessModule.cpp
+++ b/Engine/Plugins/Developer/N10XSourceCodeAccess/Source/N10XSourceCodeAccess/Private/N10XSourceCodeAccessModule.cpp
@@ -3,6 +3,7 @@
 #include "N10XSourceCodeAccessModule.h"
 
 #include "Modules/ModuleManager.h"
+#include "Features/IModularFeatures.h"
 #include "N10XSourceCodeAccessor.h"
 
 IMPLEMENT_MODULE( F10XSourceCodeAccessModule, N10XSourceCodeAccess );
