[UE] Fix 10x source code accessor to pull the correct solution file name

diff --git a/Engine/Plugins/Developer/N10XSourceCodeAccess/Source/N10XSourceCodeAccess/N10XSourceCodeAccess.Build.cs b/Engine/Plugins/Developer/N10XSourceCodeAccess/Source/N10XSourceCodeAccess/N10XSourceCodeAccess.Build.cs
--- a/Engine/Plugins/Developer/N10XSourceCodeAccess/Source/N10XSourceCodeAccess/N10XSourceCodeAccess.Build.cs
+++ b/Engine/Plugins/Developer/N10XSourceCodeAccess/Source/N10XSourceCodeAccess/N10XSourceCodeAccess.Build.cs
@@ -14,7 +14,8 @@ namespace UnrealBuildTool.Rules
 				{
 					"Core",
 					"SourceCodeAccess",
-					"DesktopPlatform"
+					"DesktopPlatform",
+					"Projects"
 				}
 			);
 
diff --git a/Engine/Plugins/Developer/N10XSourceCodeAccess/Source/N10XSourceCodeAccess/Private/N10XSourceCodeAccessor.cpp b/Engine/Plugins/Developer/N10XSourceCodeAccess/Source/N10XSourceCodeAccess/Private/N10XSourceCodeAccessor.cpp
--- a/Engine/Plugins/Developer/N10XSourceCodeAccess/Source/N10XSourceCodeAccess/Private/N10XSourceCodeAccessor.cpp
+++ b/Engine/Plugins/Developer/N10XSourceCodeAccess/Source/N10XSourceCodeAccess/Private/N10XSourceCodeAccessor.cpp
@@ -1,9 +1,13 @@
 // Copyright Epic Games, Inc. All Rights Reserved.
 
 #include "N10XSourceCodeAccessor.h"
+
+#include "Interfaces/IProjectManager.h"
+#include "Misc/FileHelper.h"
 #include "Misc/Paths.h"
 #include "Misc/UProjectInfo.h"
 #include "Misc/App.h"
+#include "ProjectDescriptor.h"
 
 #if PLATFORM_WINDOWS
 #include "Internationalization/Regex.h"
@@ -25,17 +29,23 @@ const FString& F10XSourceCodeAccessor::GetSolutionPath() const
 
 	if (CachedSolutionPath.IsEmpty() && IsInGameThread())
 	{
-		CachedSolutionPath = FPaths::ProjectDir();
-
-		if (!FUProjectDictionary::GetDefault().IsForeignProject(CachedSolutionPath))
-		{
-			CachedSolutionPath = FPaths::Combine(FPaths::RootDir(), FString("UE5") + G10XWorkspaceExtension);
-		}
-		else
+		FString PrimaryProjectPath;
+		if (!FFileHelper::LoadFileToString(PrimaryProjectPath, *(FPaths::EngineIntermediateDir() / TEXT("ProjectFiles") / TEXT("PrimaryProjectPath.txt"))))
 		{
-			FString BaseName = FApp::HasProjectName() ? FApp::GetProjectName() : FPaths::GetBaseFilename(CachedSolutionPath);
-			CachedSolutionPath = FPaths::Combine(CachedSolutionPath, BaseName + G10XWorkspaceExtension);
+			const FProjectDescriptor* CurrentProject = IProjectManager::Get().GetCurrentProject();
+			const FString ProjectDir = FPaths::ConvertRelativePathToFull(FPaths::ProjectDir());
+
+			if ((CurrentProject == nullptr || CurrentProject->Modules.Num() == 0) || !FUProjectDictionary::GetDefault().IsForeignProject(ProjectDir))
+			{
+				PrimaryProjectPath = FPaths::RootDir() / TEXT("UE5");
+			}
+			else
+			{
+				const FString BaseName = FApp::HasProjectName() ? FApp::GetProjectName() : FPaths::GetBaseFilename(ProjectDir);
+				PrimaryProjectPath = ProjectDir / BaseName;
+			}
 		}
+		CachedSolutionPath = PrimaryProjectPath + TEXT(".sln");
 	}
 
 	return CachedSolutionPath;
