PR #13018: 10x Source Code Editor: Fix files not opening when file path contains spaces

diff --git a/Engine/Plugins/Developer/N10XSourceCodeAccess/Source/N10XSourceCodeAccess/Private/N10XSourceCodeAccessor.cpp b/Engine/Plugins/Developer/N10XSourceCodeAccess/Source/N10XSourceCodeAccess/Private/N10XSourceCodeAccessor.cpp
--- a/Engine/Plugins/Developer/N10XSourceCodeAccess/Source/N10XSourceCodeAccess/Private/N10XSourceCodeAccessor.cpp
+++ b/Engine/Plugins/Developer/N10XSourceCodeAccess/Source/N10XSourceCodeAccess/Private/N10XSourceCodeAccessor.cpp
@@ -209,7 +209,7 @@ bool F10XSourceCodeAccessor::Launch(const TArray<FString>& InArgs)
 		FString ArgsString;
 		for (const FString& Arg : InArgs)
 		{
-			ArgsString.Append(Arg);
+			ArgsString.Append(FString::Printf(TEXT("\"%s\""), *Arg));
 			ArgsString.Append(TEXT(" "));
 		}
 
